Fix missing language declaration in resource file; causing all string resources to be missing.

#rnx

diff --git a/Engine/Build/Windows/Default.rc2 b/Engine/Build/Windows/Default.rc2
--- a/Engine/Build/Windows/Default.rc2
+++ b/Engine/Build/Windows/Default.rc2
@@ -6,6 +6,17 @@
 #include "../../Source/Runtime/Core/Public/Misc/CoreMiscDefines.h"
 #include "../../Source/Runtime/Launch/Resources/Version.h"
 
+#define APSTUDIO_READONLY_SYMBOLS
+#include <windows.h>
+#undef APSTUDIO_READONLY_SYMBOLS
+
+/////////////////////////////////////////////////////////////////////////////
+// English (United States) resources
+
+#if !defined(AFX_RESOURCE_DLL) || defined(AFX_TARG_ENU)
+LANGUAGE LANG_ENGLISH, SUBLANG_ENGLISH_US
+#pragma code_page(1252)
+
 /////////////////////////////////////////////////////////////////////////////
 //
 // Version
@@ -84,3 +95,6 @@ IDICON_UE4Game          ICON                    BUILD_ICON_FILE_NAME
 
 #endif
 
+#endif    // English (United States) resources
+/////////////////////////////////////////////////////////////////////////////
+
