Move Windows resource files under a Resources subfolder.

#rnx

diff --git a/Engine/Build/Windows/Default-Win32-Debug.manifest b/Engine/Build/Windows/Default-Win32-Debug.manifest
deleted file mode 100644
diff --git a/Engine/Build/Windows/Default-Win32.manifest b/Engine/Build/Windows/Default-Win32.manifest
deleted file mode 100644
diff --git a/Engine/Build/Windows/Default-Win64-Debug.manifest b/Engine/Build/Windows/Default-Win64-Debug.manifest
deleted file mode 100644
diff --git a/Engine/Build/Windows/Default-Win64.manifest b/Engine/Build/Windows/Default-Win64.manifest
deleted file mode 100644
diff --git a/Engine/Build/Windows/Default.rc2 b/Engine/Build/Windows/Default.rc2
deleted file mode 100644
diff --git a/Engine/Build/Windows/Resources/Default-Win32-Debug.manifest b/Engine/Build/Windows/Resources/Default-Win32-Debug.manifest
new file mode 100644
--- /dev/null
+++ b/Engine/Build/Windows/Resources/Default-Win32-Debug.manifest
@@ -0,0 +1,43 @@
+<?xml version="1.0" encoding="UTF-8" standalone="yes"?>
+<assembly xmlns="urn:schemas-microsoft-com:asm.v1" manifestVersion="1.0" xmlns:asmv3="urn:schemas-microsoft-com:asm.v3">
+
+  <!-- Tell Vista we don't need administrator level access-->
+  <ms_asmv2:trustInfo xmlns:ms_asmv2="urn:schemas-microsoft-com:asm.v2">
+    <ms_asmv2:security>
+      <ms_asmv2:requestedPrivileges>
+        <ms_asmv2:requestedExecutionLevel level="asInvoker" uiAccess="false" />
+      </ms_asmv2:requestedPrivileges>
+    </ms_asmv2:security>
+  </ms_asmv2:trustInfo>
+
+  <!-- Use themed common controls -->
+  <dependency>
+    <dependentAssembly>
+      <assemblyIdentity type='win32' name='Microsoft.Windows.Common-Controls' version='6.0.0.0' processorArchitecture='amd64' publicKeyToken='6595b64144ccf1df' />
+    </dependentAssembly>
+  </dependency>
+
+  <!-- OS Support -->
+  <compatibility xmlns="urn:schemas-microsoft-com:compatibility.v1">
+    <application>
+      <!-- Windows 10 -->
+      <supportedOS Id="{8e0f7a12-bfb3-4fe8-b9a5-48fd50a15a9a}"/>
+      <!-- Windows 8.1 -->
+      <supportedOS Id="{1f676c76-80e1-4239-95bb-83d0f6d0da78}"/>
+      <!-- Windows Vista -->
+      <supportedOS Id="{e2011457-1546-43c5-a5fe-008deee3d3f0}"/>
+      <!-- Windows 7 -->
+      <supportedOS Id="{35138b9a-5d96-4fbd-8e2d-a2440225f93a}"/>
+      <!-- Windows 8 -->
+      <supportedOS Id="{4a2f28e3-53b9-4441-ba9c-d69d4a4a6e38}"/>
+    </application>
+  </compatibility>
+
+  <!-- Allow paths longer than MAX_PATH (also requires enabling per-system via group policy or registry entry) -->
+  <asmv3:application>
+    <windowsSettings xmlns:ws2="http://schemas.microsoft.com/SMI/2016/WindowsSettings">
+      <ws2:longPathAware>true</ws2:longPathAware>
+    </windowsSettings>
+  </asmv3:application>
+
+</assembly>
diff --git a/Engine/Build/Windows/Resources/Default-Win32.manifest b/Engine/Build/Windows/Resources/Default-Win32.manifest
new file mode 100644
--- /dev/null
+++ b/Engine/Build/Windows/Resources/Default-Win32.manifest
@@ -0,0 +1,43 @@
+﻿<?xml version="1.0" encoding="UTF-8" standalone="yes"?>
+<assembly xmlns="urn:schemas-microsoft-com:asm.v1" manifestVersion="1.0" xmlns:asmv3="urn:schemas-microsoft-com:asm.v3">
+
+  <!-- Tell Vista we don't need administrator level access-->
+  <ms_asmv2:trustInfo xmlns:ms_asmv2="urn:schemas-microsoft-com:asm.v2">
+    <ms_asmv2:security>
+      <ms_asmv2:requestedPrivileges>
+        <ms_asmv2:requestedExecutionLevel level="asInvoker" uiAccess="false" />
+      </ms_asmv2:requestedPrivileges>
+    </ms_asmv2:security>
+  </ms_asmv2:trustInfo>
+
+  <!-- Use themed common controls -->
+  <dependency>
+    <dependentAssembly>
+      <assemblyIdentity type='win32' name='Microsoft.Windows.Common-Controls' version='6.0.0.0' processorArchitecture='amd64' publicKeyToken='6595b64144ccf1df' />
+    </dependentAssembly>
+  </dependency>
+
+  <!-- OS Support -->
+  <compatibility xmlns="urn:schemas-microsoft-com:compatibility.v1">
+    <application>
+      <!-- Windows 10 -->
+      <supportedOS Id="{8e0f7a12-bfb3-4fe8-b9a5-48fd50a15a9a}"/>
+      <!-- Windows 8.1 -->
+      <supportedOS Id="{1f676c76-80e1-4239-95bb-83d0f6d0da78}"/>
+      <!-- Windows Vista -->
+      <supportedOS Id="{e2011457-1546-43c5-a5fe-008deee3d3f0}"/>
+      <!-- Windows 7 -->
+      <supportedOS Id="{35138b9a-5d96-4fbd-8e2d-a2440225f93a}"/>
+      <!-- Windows 8 -->
+      <supportedOS Id="{4a2f28e3-53b9-4441-ba9c-d69d4a4a6e38}"/>
+    </application>
+  </compatibility>
+
+  <!-- Allow paths longer than MAX_PATH (also requires enabling per-system via group policy or registry entry) -->
+  <asmv3:application>
+    <windowsSettings xmlns:ws2="http://schemas.microsoft.com/SMI/2016/WindowsSettings">
+      <ws2:longPathAware>true</ws2:longPathAware>
+    </windowsSettings>
+  </asmv3:application>
+
+</assembly>
diff --git a/Engine/Build/Windows/Resources/Default-Win64-Debug.manifest b/Engine/Build/Windows/Resources/Default-Win64-Debug.manifest
new file mode 100644
--- /dev/null
+++ b/Engine/Build/Windows/Resources/Default-Win64-Debug.manifest
@@ -0,0 +1,43 @@
+<?xml version="1.0" encoding="UTF-8" standalone="yes"?>
+<assembly xmlns="urn:schemas-microsoft-com:asm.v1" manifestVersion="1.0" xmlns:asmv3="urn:schemas-microsoft-com:asm.v3">
+
+  <!-- Tell Vista we don't need administrator level access-->
+  <ms_asmv2:trustInfo xmlns:ms_asmv2="urn:schemas-microsoft-com:asm.v2">
+    <ms_asmv2:security>
+      <ms_asmv2:requestedPrivileges>
+        <ms_asmv2:requestedExecutionLevel level="asInvoker" uiAccess="false" />
+      </ms_asmv2:requestedPrivileges>
+    </ms_asmv2:security>
+  </ms_asmv2:trustInfo>
+
+  <!-- Use themed common controls -->
+  <dependency>
+    <dependentAssembly>
+      <assemblyIdentity type='win32' name='Microsoft.Windows.Common-Controls' version='6.0.0.0' processorArchitecture='amd64' publicKeyToken='6595b64144ccf1df' />
+    </dependentAssembly>
+  </dependency>
+
+  <!-- OS Support -->
+  <compatibility xmlns="urn:schemas-microsoft-com:compatibility.v1">
+    <application>
+      <!-- Windows 10 -->
+      <supportedOS Id="{8e0f7a12-bfb3-4fe8-b9a5-48fd50a15a9a}"/>
+      <!-- Windows 8.1 -->
+      <supportedOS Id="{1f676c76-80e1-4239-95bb-83d0f6d0da78}"/>
+      <!-- Windows Vista -->
+      <supportedOS Id="{e2011457-1546-43c5-a5fe-008deee3d3f0}"/>
+      <!-- Windows 7 -->
+      <supportedOS Id="{35138b9a-5d96-4fbd-8e2d-a2440225f93a}"/>
+      <!-- Windows 8 -->
+      <supportedOS Id="{4a2f28e3-53b9-4441-ba9c-d69d4a4a6e38}"/>
+    </application>
+  </compatibility>
+
+  <!-- Allow paths longer than MAX_PATH (also requires enabling per-system via group policy or registry entry) -->
+  <asmv3:application>
+    <windowsSettings xmlns:ws2="http://schemas.microsoft.com/SMI/2016/WindowsSettings">
+      <ws2:longPathAware>true</ws2:longPathAware>
+    </windowsSettings>
+  </asmv3:application>
+
+</assembly>
diff --git a/Engine/Build/Windows/Resources/Default-Win64.manifest b/Engine/Build/Windows/Resources/Default-Win64.manifest
new file mode 100644
--- /dev/null
+++ b/Engine/Build/Windows/Resources/Default-Win64.manifest
@@ -0,0 +1,43 @@
+<?xml version="1.0" encoding="UTF-8" standalone="yes"?>
+<assembly xmlns="urn:schemas-microsoft-com:asm.v1" manifestVersion="1.0" xmlns:asmv3="urn:schemas-microsoft-com:asm.v3">
+
+  <!-- Tell Vista we don't need administrator level access-->
+  <ms_asmv2:trustInfo xmlns:ms_asmv2="urn:schemas-microsoft-com:asm.v2">
+    <ms_asmv2:security>
+      <ms_asmv2:requestedPrivileges>
+        <ms_asmv2:requestedExecutionLevel level="asInvoker" uiAccess="false" />
+      </ms_asmv2:requestedPrivileges>
+    </ms_asmv2:security>
+  </ms_asmv2:trustInfo>
+
+  <!-- Use themed common controls -->
+  <dependency>
+    <dependentAssembly>
+      <assemblyIdentity type='win32' name='Microsoft.Windows.Common-Controls' version='6.0.0.0' processorArchitecture='amd64' publicKeyToken='6595b64144ccf1df' />
+    </dependentAssembly>
+  </dependency>
+
+  <!-- OS Support -->
+  <compatibility xmlns="urn:schemas-microsoft-com:compatibility.v1">
+    <application>
+      <!-- Windows 10 -->
+      <supportedOS Id="{8e0f7a12-bfb3-4fe8-b9a5-48fd50a15a9a}"/>
+      <!-- Windows 8.1 -->
+      <supportedOS Id="{1f676c76-80e1-4239-95bb-83d0f6d0da78}"/>
+      <!-- Windows Vista -->
+      <supportedOS Id="{e2011457-1546-43c5-a5fe-008deee3d3f0}"/>
+      <!-- Windows 7 -->
+      <supportedOS Id="{35138b9a-5d96-4fbd-8e2d-a2440225f93a}"/>
+      <!-- Windows 8 -->
+      <supportedOS Id="{4a2f28e3-53b9-4441-ba9c-d69d4a4a6e38}"/>
+    </application>
+  </compatibility>
+
+  <!-- Allow paths longer than MAX_PATH (also requires enabling per-system via group policy or registry entry) -->
+  <asmv3:application>
+    <windowsSettings xmlns:ws2="http://schemas.microsoft.com/SMI/2016/WindowsSettings">
+      <ws2:longPathAware>true</ws2:longPathAware>
+    </windowsSettings>
+  </asmv3:application>
+
+</assembly>
diff --git a/Engine/Build/Windows/Resources/Default.rc2 b/Engine/Build/Windows/Resources/Default.rc2
new file mode 100644
--- /dev/null
+++ b/Engine/Build/Windows/Resources/Default.rc2
@@ -0,0 +1,125 @@
+// Copyright 1998-2019 Epic Games, Inc. All Rights Reserved.
+//
+// This file contains resources which are edited by hand, and cannot be loaded/saved by the Visual Studio resource editor.
+
+#include "../../../Source/Runtime/Launch/Resources/Windows/resource.h"
+#include "../../../Source/Runtime/Core/Public/Misc/CoreMiscDefines.h"
+#include "../../../Source/Runtime/Launch/Resources/Version.h"
+
+#define APSTUDIO_READONLY_SYMBOLS
+#include <windows.h>
+#undef APSTUDIO_READONLY_SYMBOLS
+
+// Various strings used for project resources
+#ifdef PROJECT_COMPANY_NAME
+	#define BUILD_PROJECT_COMPANY_NAME PREPROCESSOR_TO_STRING(PROJECT_COMPANY_NAME)
+#else
+	#define BUILD_PROJECT_COMPANY_NAME EPIC_COMPANY_NAME
+#endif
+
+#ifdef PROJECT_COPYRIGHT_STRING
+	#define BUILD_PROJECT_COPYRIGHT_STRING PREPROCESSOR_TO_STRING(PROJECT_COPYRIGHT_STRING)
+#else
+	#define BUILD_PROJECT_COPYRIGHT_STRING EPIC_COPYRIGHT_STRING
+#endif
+
+#ifdef PROJECT_PRODUCT_NAME
+	#define BUILD_PROJECT_PRODUCT_NAME PREPROCESSOR_TO_STRING(PROJECT_PRODUCT_NAME)
+#else
+	#define BUILD_PROJECT_PRODUCT_NAME PREPROCESSOR_TO_STRING(UE_APP_NAME)
+#endif
+
+#ifdef PROJECT_PRODUCT_IDENTIFIER
+	#define BUILD_PROJECT_PRODUCT_IDENTIFIER PREPROCESSOR_TO_STRING(PROJECT_PRODUCT_IDENTIFIER)
+#else
+	#define BUILD_PROJECT_PRODUCT_IDENTIFIER EPIC_PRODUCT_IDENTIFIER
+#endif
+
+/////////////////////////////////////////////////////////////////////////////
+// English (United States) resources
+
+#if !defined(AFX_RESOURCE_DLL) || defined(AFX_TARG_ENU)
+LANGUAGE LANG_ENGLISH, SUBLANG_ENGLISH_US
+#pragma code_page(1252)
+
+/////////////////////////////////////////////////////////////////////////////
+//
+// Version
+//
+
+VS_VERSION_INFO VERSIONINFO
+ FILEVERSION ENGINE_MAJOR_VERSION,ENGINE_MINOR_VERSION,ENGINE_PATCH_VERSION,0
+ PRODUCTVERSION ENGINE_MAJOR_VERSION,ENGINE_MINOR_VERSION,ENGINE_PATCH_VERSION,0
+ FILEFLAGSMASK 0x17L
+#ifdef _DEBUG
+ FILEFLAGS 0x1L
+#else
+ FILEFLAGS 0x0L
+#endif
+ FILEOS 0x4L
+ FILETYPE 0x2L
+ FILESUBTYPE 0x0L
+BEGIN
+	BLOCK "StringFileInfo"
+	BEGIN
+		BLOCK "040904b0"
+		BEGIN
+			VALUE "CompanyName", BUILD_PROJECT_COMPANY_NAME
+			VALUE "LegalCopyright", BUILD_PROJECT_COPYRIGHT_STRING
+			VALUE "ProductName", BUILD_PROJECT_PRODUCT_NAME
+#ifdef BUILD_VERSION
+			VALUE "ProductVersion", PREPROCESSOR_TO_STRING(BUILD_VERSION)
+#endif
+			VALUE "FileDescription", BUILD_PROJECT_PRODUCT_NAME
+			VALUE "InternalName", BUILD_PROJECT_PRODUCT_IDENTIFIER
+#ifdef ORIGINAL_FILE_NAME
+			VALUE "OriginalFilename", PREPROCESSOR_TO_STRING(ORIGINAL_FILE_NAME)
+#endif
+		END
+	END
+	BLOCK "VarFileInfo"
+	BEGIN
+		VALUE "Translation", 0x409, 1200
+	END
+END
+
+/////////////////////////////////////////////////////////////////////////////
+//
+// Icon
+//
+
+// Icon with lowest ID value placed first to ensure application icon
+// remains consistent on all systems.
+
+#ifdef BUILD_ICON_FILE_NAME
+IDICON_UE4Game          ICON                    BUILD_ICON_FILE_NAME
+#endif
+
+/////////////////////////////////////////////////////////////////////////////
+//
+// Application Manifest
+//
+
+#if _DEBUG
+
+// Debug manifest
+#if _WIN64
+1						RT_MANIFEST				"Default-Win64-Debug.manifest"
+#else
+1						RT_MANIFEST				"Default-Win32-Debug.manifest"
+#endif
+
+#else
+
+// Release manifest
+#if _WIN64
+1						RT_MANIFEST				"Default-Win64.manifest"
+#else
+1						RT_MANIFEST				"Default-Win32.manifest"
+#endif
+
+#endif
+
+#endif    // English (United States) resources
+/////////////////////////////////////////////////////////////////////////////
+
diff --git a/Engine/Source/Programs/UnrealBuildTool/Configuration/UEBuildBinary.cs b/Engine/Source/Programs/UnrealBuildTool/Configuration/UEBuildBinary.cs
--- a/Engine/Source/Programs/UnrealBuildTool/Configuration/UEBuildBinary.cs
+++ b/Engine/Source/Programs/UnrealBuildTool/Configuration/UEBuildBinary.cs
@@ -743,7 +743,7 @@ namespace UnrealBuildTool
 						ResourceCompileEnvironment.Definitions.Add(String.Format("BUILD_VERSION={0}", Target.BuildVersion));
 
 						// Otherwise compile the default resource file per-binary, so that it gets the correct ORIGINAL_FILE_NAME macro.
-						FileItem DefaultResourceFile = FileItem.GetItemByFileReference(FileReference.Combine(UnrealBuildTool.EngineDirectory, "Build", "Windows", "Default.rc2"));
+						FileItem DefaultResourceFile = FileItem.GetItemByFileReference(FileReference.Combine(UnrealBuildTool.EngineDirectory, "Build", "Windows", "Resources", "Default.rc2"));
 						CPPOutput DefaultResourceOutput = ToolChain.CompileRCFiles(ResourceCompileEnvironment, new List<FileItem> { DefaultResourceFile }, ResourceIntermediateDirectory, Makefile.Actions);
 						BinaryLinkEnvironment.InputFiles.AddRange(DefaultResourceOutput.ObjectFiles);
 					}
diff --git a/Engine/Source/Programs/UnrealBuildTool/Configuration/UEBuildTarget.cs b/Engine/Source/Programs/UnrealBuildTool/Configuration/UEBuildTarget.cs
--- a/Engine/Source/Programs/UnrealBuildTool/Configuration/UEBuildTarget.cs
+++ b/Engine/Source/Programs/UnrealBuildTool/Configuration/UEBuildTarget.cs
@@ -1281,7 +1281,7 @@ namespace UnrealBuildTool
 					{
 						CppCompileEnvironment DefaultResourceCompileEnvironment = new CppCompileEnvironment(GlobalCompileEnvironment);
 
-						FileItem DefaultResourceFile = FileItem.GetItemByFileReference(FileReference.Combine(UnrealBuildTool.EngineDirectory, "Build", "Windows", "Default.rc2"));
+						FileItem DefaultResourceFile = FileItem.GetItemByFileReference(FileReference.Combine(UnrealBuildTool.EngineDirectory, "Build", "Windows", "Resources", "Default.rc2"));
 
 						WindowsPlatform.SetupResourceCompileEnvironment(DefaultResourceCompileEnvironment, Rules);
 
diff --git a/Engine/Source/Programs/UnrealBuildTool/Platform/Windows/UEBuildWindows.cs b/Engine/Source/Programs/UnrealBuildTool/Platform/Windows/UEBuildWindows.cs
--- a/Engine/Source/Programs/UnrealBuildTool/Platform/Windows/UEBuildWindows.cs
+++ b/Engine/Source/Programs/UnrealBuildTool/Platform/Windows/UEBuildWindows.cs
@@ -1466,7 +1466,7 @@ namespace UnrealBuildTool
 			}
 
 			// Otherwise use the default
-			return FileReference.Combine(UnrealBuildTool.EngineDirectory, "Build", "Windows", "Default.ico");
+			return FileReference.Combine(UnrealBuildTool.EngineDirectory, "Build", "Windows", "Resources", "Default.ico");
 		}
 
 		/// <summary>
