Fix manifest on Win32 referencing amd64 version of common controls assembly causing executables to fail to start. Not sure why this is a regression; manifest from 4.21 (//UE4/Release-4.21/Engine/Source/Runtime/Launch/Resources/Windows/PCLaunch-Win32.manifest) also references the incorrect platform.

diff --git a/Engine/Build/Windows/Resources/Default-Win32-Debug.manifest b/Engine/Build/Windows/Resources/Default-Win32-Debug.manifest
--- a/Engine/Build/Windows/Resources/Default-Win32-Debug.manifest
+++ b/Engine/Build/Windows/Resources/Default-Win32-Debug.manifest
@@ -13,8 +13,8 @@
   <!-- Use themed common controls -->
   <dependency>
     <dependentAssembly>
-      <assemblyIdentity type='win32' name='Microsoft.Windows.Common-Controls' version='6.0.0.0' processorArchitecture='amd64' publicKeyToken='6595b64144ccf1df' />
-    </dependentAssembly>
+      <assemblyIdentity type='win32' name='Microsoft.Windows.Common-Controls' version='6.0.0.0' processorArchitecture='x86' publicKeyToken='6595b64144ccf1df' />
+	</dependentAssembly>
   </dependency>
 
   <!-- OS Support -->
diff --git a/Engine/Build/Windows/Resources/Default-Win32.manifest b/Engine/Build/Windows/Resources/Default-Win32.manifest
--- a/Engine/Build/Windows/Resources/Default-Win32.manifest
+++ b/Engine/Build/Windows/Resources/Default-Win32.manifest
@@ -1,4 +1,4 @@
-﻿<?xml version="1.0" encoding="UTF-8" standalone="yes"?>
+<?xml version="1.0" encoding="UTF-8" standalone="yes"?>
 <assembly xmlns="urn:schemas-microsoft-com:asm.v1" manifestVersion="1.0" xmlns:asmv3="urn:schemas-microsoft-com:asm.v3">
 
   <!-- Tell Vista we don't need administrator level access-->
@@ -13,11 +13,11 @@
   <!-- Use themed common controls -->
   <dependency>
     <dependentAssembly>
-      <assemblyIdentity type='win32' name='Microsoft.Windows.Common-Controls' version='6.0.0.0' processorArchitecture='amd64' publicKeyToken='6595b64144ccf1df' />
+      <assemblyIdentity type='win32' name='Microsoft.Windows.Common-Controls' version='6.0.0.0' processorArchitecture='x86' publicKeyToken='6595b64144ccf1df' />
     </dependentAssembly>
   </dependency>
 
-  <!-- OS Support -->
+	<!-- OS Support -->
   <compatibility xmlns="urn:schemas-microsoft-com:compatibility.v1">
     <application>
       <!-- Windows 10 -->
