Only include what the rc2 actually needs
#rnx

diff --git a/Engine/Build/Windows/Resources/Default.rc2 b/Engine/Build/Windows/Resources/Default.rc2
--- a/Engine/Build/Windows/Resources/Default.rc2
+++ b/Engine/Build/Windows/Resources/Default.rc2
@@ -3,7 +3,7 @@
 // This file contains resources which are edited by hand, and cannot be loaded/saved by the Visual Studio resource editor.
 
 #include "../../../Source/Runtime/Launch/Resources/Windows/resource.h"
-#include "../../../Source/Runtime/Core/Public/Misc/CoreMiscDefines.h"
+#include "../../../Source/Runtime/Core/Public/HAL/PreprocessorHelpers.h"
 #include "../../../Source/Runtime/Launch/Resources/Version.h"
 
 #define APSTUDIO_READONLY_SYMBOLS
