[Windows]
* Added "FileVersion" to Default.rc2... this was needed in order for powershell and c# to be able to query version information. Without that all fields ended up empty (but still showed correct in file properties -> details)

diff --git a/Engine/Build/Windows/Resources/Default.rc2 b/Engine/Build/Windows/Resources/Default.rc2
--- a/Engine/Build/Windows/Resources/Default.rc2
+++ b/Engine/Build/Windows/Resources/Default.rc2
@@ -70,6 +70,7 @@ BEGIN
 	BEGIN
 		BLOCK "040904b0"
 		BEGIN
+			VALUE "FileVersion", BUILD_PROJECT_PRODUCT_VERSION
 			VALUE "CompanyName", BUILD_PROJECT_COMPANY_NAME
 			VALUE "LegalCopyright", BUILD_PROJECT_COPYRIGHT_STRING
 			VALUE "ProductName", BUILD_PROJECT_PRODUCT_NAME
